Added a toolset for animation features.  This plugin will contain toolsets for systems such as control rig, animation blueprints, and the new unreal animation framework.

This first commit include basic MVP control rig support and tests for control rig.

[RN] Added an assistant toolset for animation features such as control rig
#Synced-CL 51596834

diff --git a/Engine/Plugins/Experimental/AnimationAssistantToolset/AnimationAssistantToolset.uplugin b/Engine/Plugins/Experimental/AnimationAssistantToolset/AnimationAssistantToolset.uplugin
new file mode 100644
--- /dev/null
+++ b/Engine/Plugins/Experimental/AnimationAssistantToolset/AnimationAssistantToolset.uplugin
@@ -0,0 +1,34 @@
+{
+	"FileVersion": 3,
+	"Version": 1,
+	"VersionName": "1.0",
+	"FriendlyName": "AnimationAssistantToolset",
+	"Description": "Toolset for Animation Systems",
+	"Category": "Other",
+	"CreatedBy": "Epic Games",
+	"CreatedByURL": "",
+	"DocsURL": "",
+	"MarketplaceURL": "",
+	"SupportURL": "",
+	"CanContainContent": true,
+	"IsBetaVersion": false,
+	"IsExperimentalVersion": true,
+	"Installed": false,
+	"Modules": [
+		{
+			"Name": "AnimationAssistantToolset",
+			"Type": "Runtime",
+			"LoadingPhase": "Default"
+		}
+	],
+	"Plugins": [
+		{
+			"Name": "ControlRig",
+			"Enabled": true
+		},
+		{
+			"Name": "ToolsetRegistry",
+			"Enabled": true
+		}
+	]
+}
\ No newline at end of file
diff --git a/Engine/Plugins/Experimental/AnimationAssistantToolset/Source/AnimationAssistantToolset/AnimationAssistantToolset.Build.cs b/Engine/Plugins/Experimental/AnimationAssistantToolset/Source/AnimationAssistantToolset/AnimationAssistantToolset.Build.cs
new file mode 100644
--- /dev/null
+++ b/Engine/Plugins/Experimental/AnimationAssistantToolset/Source/AnimationAssistantToolset/AnimationAssistantToolset.Build.cs
@@ -0,0 +1,53 @@
+// Copyright Epic Games, Inc. All Rights Reserved.
+
+using UnrealBuildTool;
+
+public class AnimationAssistantToolset : ModuleRules
+{
+	public AnimationAssistantToolset(ReadOnlyTargetRules Target) : base(Target)
+	{
+		PCHUsage = ModuleRules.PCHUsageMode.UseExplicitOrSharedPCHs;
+		
+		PublicIncludePaths.AddRange(
+			new string[] {
+				// ... add public include paths required here ...
+			}
+			);
+				
+		
+		PrivateIncludePaths.AddRange(
+			new string[] {
+				// ... add other private include paths required here ...
+			}
+			);
+			
+		
+		PublicDependencyModuleNames.AddRange(
+			new string[]
+			{
+				"Core",
+				// ... add other public dependencies that you statically link with here ...
+			}
+			);
+			
+		
+		PrivateDependencyModuleNames.AddRange(
+			new string[]
+			{
+				"CoreUObject",
+				"Engine",
+				"Slate",
+				"SlateCore",
+				// ... add private dependencies that you statically link with here ...	
+			}
+			);
+		
+		
+		DynamicallyLoadedModuleNames.AddRange(
+			new string[]
+			{
+				// ... add any modules that your module loads dynamically here ...
+			}
+			);
+	}
+}
diff --git a/Engine/Plugins/Experimental/AnimationAssistantToolset/Source/AnimationAssistantToolset/Private/AnimationAssistantToolset.cpp b/Engine/Plugins/Experimental/AnimationAssistantToolset/Source/AnimationAssistantToolset/Private/AnimationAssistantToolset.cpp
new file mode 100644
--- /dev/null
+++ b/Engine/Plugins/Experimental/AnimationAssistantToolset/Source/AnimationAssistantToolset/Private/AnimationAssistantToolset.cpp
@@ -0,0 +1,20 @@
+// Copyright Epic Games, Inc. All Rights Reserved.
+
+#include "AnimationAssistantToolset.h"
+
+#define LOCTEXT_NAMESPACE "FAnimationAssistantToolsetModule"
+
+void FAnimationAssistantToolsetModule::StartupModule()
+{
+	// This code will execute after your module is loaded into memory; the exact timing is specified in the .uplugin file per-module
+}
+
+void FAnimationAssistantToolsetModule::ShutdownModule()
+{
+	// This function may be called during shutdown to clean up your module.  For modules that support dynamic reloading,
+	// we call this function before unloading the module.
+}
+
+#undef LOCTEXT_NAMESPACE
+	
+IMPLEMENT_MODULE(FAnimationAssistantToolsetModule, AnimationAssistantToolset)
\ No newline at end of file
diff --git a/Engine/Plugins/Experimental/AnimationAssistantToolset/Source/AnimationAssistantToolset/Public/AnimationAssistantToolset.h b/Engine/Plugins/Experimental/AnimationAssistantToolset/Source/AnimationAssistantToolset/Public/AnimationAssistantToolset.h
new file mode 100644
--- /dev/null
+++ b/Engine/Plugins/Experimental/AnimationAssistantToolset/Source/AnimationAssistantToolset/Public/AnimationAssistantToolset.h
@@ -0,0 +1,14 @@
+// Copyright Epic Games, Inc. All Rights Reserved.
+
+#pragma once
+
+#include "Modules/ModuleManager.h"
+
+class FAnimationAssistantToolsetModule : public IModuleInterface
+{
+public:
+
+	/** IModuleInterface implementation */
+	virtual void StartupModule() override;
+	virtual void ShutdownModule() override;
+};
