[AI Toolsets]: Ensure all toolset plugins are marked as editor only.

#Synced-CL 52422229

diff --git a/Engine/Plugins/Experimental/Toolsets/AnimationAssistantToolset/AnimationAssistantToolset.uplugin b/Engine/Plugins/Experimental/Toolsets/AnimationAssistantToolset/AnimationAssistantToolset.uplugin
--- a/Engine/Plugins/Experimental/Toolsets/AnimationAssistantToolset/AnimationAssistantToolset.uplugin
+++ b/Engine/Plugins/Experimental/Toolsets/AnimationAssistantToolset/AnimationAssistantToolset.uplugin
@@ -10,6 +10,7 @@
 	"DocsURL": "",
 	"MarketplaceURL": "",
 	"SupportURL": "",
+	"EditorOnly": true,
 	"CanContainContent": true,
 	"IsBetaVersion": false,
 	"IsExperimentalVersion": true,
diff --git a/Engine/Plugins/Experimental/Toolsets/MCPClientToolset/MCPClientToolset.uplugin b/Engine/Plugins/Experimental/Toolsets/MCPClientToolset/MCPClientToolset.uplugin
--- a/Engine/Plugins/Experimental/Toolsets/MCPClientToolset/MCPClientToolset.uplugin
+++ b/Engine/Plugins/Experimental/Toolsets/MCPClientToolset/MCPClientToolset.uplugin
@@ -10,6 +10,7 @@
 	"DocsURL": "",
 	"MarketplaceURL": "",
 	"SupportURL": "",
+	"EditorOnly": true,
 	"EnabledByDefault": false,
 	"CanContainContent": false,
 	"IsBetaVersion": true,
diff --git a/Engine/Plugins/Experimental/Toolsets/NiagaraAIAssistantTools/NiagaraAIAssistantTools.uplugin b/Engine/Plugins/Experimental/Toolsets/NiagaraAIAssistantTools/NiagaraAIAssistantTools.uplugin
--- a/Engine/Plugins/Experimental/Toolsets/NiagaraAIAssistantTools/NiagaraAIAssistantTools.uplugin
+++ b/Engine/Plugins/Experimental/Toolsets/NiagaraAIAssistantTools/NiagaraAIAssistantTools.uplugin
@@ -10,6 +10,7 @@
 	"DocsURL": "",
 	"MarketplaceURL": "",
 	"SupportURL": "",
+	"EditorOnly": true,
 	"EnabledByDefault": false,
 	"CanContainContent": true,
 	"IsBetaVersion": false,
diff --git a/Engine/Plugins/Experimental/Toolsets/UMGToolSet/UMGToolSet.uplugin b/Engine/Plugins/Experimental/Toolsets/UMGToolSet/UMGToolSet.uplugin
--- a/Engine/Plugins/Experimental/Toolsets/UMGToolSet/UMGToolSet.uplugin
+++ b/Engine/Plugins/Experimental/Toolsets/UMGToolSet/UMGToolSet.uplugin
@@ -10,6 +10,7 @@
 	"DocsURL": "",
 	"MarketplaceURL": "",
 	"SupportURL": "",
+	"EditorOnly": true,
 	"EnabledByDefault": false,
 	"CanContainContent": false,
 	"IsBetaVersion": false,
